[TmvMediaUtils] Movie Render Graph Tmv Encoder Integration

- Adding UMovieGraphTmvEncoderNode, derived from UMovieGraphVideoOutputNode, wrapping a Tmv Transcode Job to produce Tmv encoded sequences. Tmv Encoder is selectable (either Apv or Exr atm) and configurable.

#Synced-CL 51094493

diff --git a/Engine/Plugins/Media/ImgMedia/Source/OpenExrWrapper/Private/OpenExrWrapperTmvEncoderOptions.h b/Engine/Plugins/Media/ImgMedia/Source/OpenExrWrapper/Private/OpenExrWrapperTmvEncoderOptions.h
--- a/Engine/Plugins/Media/ImgMedia/Source/OpenExrWrapper/Private/OpenExrWrapperTmvEncoderOptions.h
+++ b/Engine/Plugins/Media/ImgMedia/Source/OpenExrWrapper/Private/OpenExrWrapperTmvEncoderOptions.h
@@ -54,10 +54,10 @@ struct FOpenExrWrapperTmvEncoderOptions : public FTmvMediaEncoderOptions
 	bool bRemoveAlphaChannel = false;
 
 	/** Tint each mip level a different colour to help with debugging. */
-	UPROPERTY(EditAnywhere, Transient, Category = Debug)
+	UPROPERTY(EditAnywhere, Category = Debug)
 	bool bEnableMipLevelTint = false;
 
 	/** Color to tint each mip level. */
-	UPROPERTY(EditAnywhere, Transient, Category = Debug)
+	UPROPERTY(EditAnywhere, Category = Debug)
 	TArray<FLinearColor> MipLevelTints;
 };
diff --git a/Engine/Plugins/Media/TmvMediaUtils/Source/MovieRenderPipelineTmvEncoder/MovieRenderPipelineTmvEncoder.Build.cs b/Engine/Plugins/Media/TmvMediaUtils/Source/MovieRenderPipelineTmvEncoder/MovieRenderPipelineTmvEncoder.Build.cs
new file mode 100644
--- /dev/null
+++ b/Engine/Plugins/Media/TmvMediaUtils/Source/MovieRenderPipelineTmvEncoder/MovieRenderPipelineTmvEncoder.Build.cs
@@ -0,0 +1,30 @@
+// Copyright Epic Games, Inc. All Rights Reserved.
+
+namespace UnrealBuildTool.Rules
+{
+	public class MovieRenderPipelineTmvEncoder : ModuleRules
+	{
+		public MovieRenderPipelineTmvEncoder(ReadOnlyTargetRules Target) : base(Target)
+		{
+            PublicDependencyModuleNames.AddRange(				
+				new string[] {
+					"MovieRenderPipelineCore",
+					"MovieRenderPipelineSettings",
+					"OpenColorIO",
+					"SlateCore",
+                });
+
+			PrivateDependencyModuleNames.AddRange(
+				new string[]
+				{
+					"Core",
+                    "CoreUObject",
+                    "Engine",
+                    "ImageCore",
+                    "RHI",
+                    "TmvMedia",
+                }
+            );
+        }
+    }
+}
diff --git a/Engine/Plugins/Media/TmvMediaUtils/Source/MovieRenderPipelineTmvEncoder/Private/Graph/MovieGraphTmvEncoderNode.cpp b/Engine/Plugins/Media/TmvMediaUtils/Source/MovieRenderPipelineTmvEncoder/Private/Graph/MovieGraphTmvEncoderNode.cpp
new file mode 100644
--- /dev/null
+++ b/Engine/Plugins/Media/TmvMediaUtils/Source/MovieRenderPipelineTmvEncoder/Private/Graph/MovieGraphTmvEncoderNode.cpp
@@ -0,0 +1,331 @@
+// Copyright Epic Games, Inc. All Rights Reserved.
+
+#include "Graph/MovieGraphTmvEncoderNode.h"
+
+#include "Encoder/TmvMediaEncoderOptions.h"
+#include "Graph/MovieGraphBlueprintLibrary.h"
+#include "Graph/MovieGraphConfig.h"
+#include "Graph/MovieGraphOCIOHelper.h"
+#include "Graph/MovieGraphPipeline.h"
+#include "Graph/Nodes/MovieGraphBurnInNode.h"
+#include "Graph/Nodes/MovieGraphGlobalOutputSettingNode.h"
+#include "ImageWriteTask.h"
+#include "MoviePipelineImageQuantization.h"
+#include "MoviePipelineTelemetry.h"
+#include "SampleBuffer.h"
+#include "SampleConverter/TmvMediaFrameMipImageBuffer.h"
+#include "SampleConverter/TmvMediaFrameMipPixelDataBuffer.h"
+#include "Styling/AppStyle.h"
+#include "Transcoder/TmvMediaFrameConverter.h"
+#include "Transcoder/TmvMediaFrameMips.h"
+#include "Transcoder/TmvMediaFrameProducer.h"
+#include "Transcoder/TmvMediaTmvFrameEncoder.h"
+#include "Transcoder/TmvMediaTranscodeJob.h"
+#include "Transcoder/TmvMediaTranscodeMuxer.h"
+
+#include UE_INLINE_GENERATED_CPP_BY_NAME(MovieGraphTmvEncoderNode)
+
+namespace UE::TmvMediaUtils::Encoder
+{
+	bool InitFrameMips(FTmvMediaFrameMips& InFrameMips, TUniquePtr<FImagePixelData>&& InMip0, bool bInGenerateMips)
+	{
+		if (!InMip0.IsValid())
+		{
+			return false;
+		}
+		
+		FImageView Mip0View = InMip0->GetImageView();	// make view before moving the pointer.
+		
+		InFrameMips.MipBuffers.Reset();
+		InFrameMips.MipBuffers.Add(MakeShared<FTmvMediaFrameMipPixelDataBuffer>(MoveTemp(InMip0), 0));
+
+		// Directly use ImageCore to generate the mips on the cpu. (It's reasonably fast.)
+		if (bInGenerateMips && (Mip0View.SizeX > 1 || Mip0View.SizeY > 1))
+		{
+			FImageView PreviousMipView = Mip0View;
+
+			// Protecting against infinite loop.
+			constexpr int32 MaxMips = 32;
+
+			while (InFrameMips.MipBuffers.Num() < MaxMips)
+			{
+				int32 MipSizeX = FMath::Max(PreviousMipView.SizeX/2, 1);
+				int32 MipSizeY = FMath::Max(PreviousMipView.SizeY/2, 1);
+			
+				TSharedPtr<FImage> NewMip = MakeShared<FImage>(MipSizeX, MipSizeY, 1, Mip0View.Format, Mip0View.GammaSpace);
+				{
+					TRACE_CPUPROFILER_EVENT_SCOPE(TmvMediaTranscode::GenerateMip);
+					FImageCore::ResizeImage(PreviousMipView, *NewMip);
+				}
+				int32 MipLevel = InFrameMips.MipBuffers.Num();			
+				InFrameMips.MipBuffers.Add(MakeShared<FTmvMediaFrameMipImageBuffer>(NewMip, MipLevel));
+
+				PreviousMipView = *NewMip;
+
+				if (MipSizeX == 1 && MipSizeY == 1)
+				{
+					break;
+				}
+			}
+		}
+		
+		return true;
+	}
+}
+
+UMovieGraphTmvEncoderNode::UMovieGraphTmvEncoderNode()
+{
+	// todo: Revisit this when tmv container is supported. For now, use file sequence format.
+	FileNameFormat = TEXT("{sequence_name}.{layer_name}.{frame_number}");
+}
+
+#if WITH_EDITOR
+FText UMovieGraphTmvEncoderNode::GetNodeTitle(const bool bGetDescriptive) const
+{
+	static const FText NodeNameNone = NSLOCTEXT("MovieGraphNodes", "NodeName_Tmv_None", "TMV Movie");
+	static const FText NodeNameFormat = NSLOCTEXT("MovieGraphNodes", "NodeName_Tmv_Format", "TMV Movie ({0})");
+
+	if (TmvEncoderOptions.IsValid())
+	{
+		return FText::Format(NodeNameFormat, FText::FromName(TmvEncoderOptions.Get().GetEncoderName()));
+	}
+	
+	return NodeNameNone;
+}
+
+FText UMovieGraphTmvEncoderNode::GetMenuCategory() const
+{
+	return NSLOCTEXT("MovieGraphNodes", "TmvEncoderNode_Category", "Output Type");
+}
+
+FText UMovieGraphTmvEncoderNode::GetKeywords() const
+{
+	static const FText Keywords = NSLOCTEXT("MovieGraphNodes", "TmvEncoderNode_Keywords", "tmv encoder apv exr movie video");
+	return Keywords;
+}
+
+FLinearColor UMovieGraphTmvEncoderNode::GetNodeTitleColor() const
+{
+	static const FLinearColor NodeColor = FLinearColor(0.047f, 0.654f, 0.537f);
+	return NodeColor;
+}
+
+FSlateIcon UMovieGraphTmvEncoderNode::GetIconAndTint(FLinearColor& OutColor) const
+{
+	static const FSlateIcon Icon = FSlateIcon(FAppStyle::GetAppStyleSetName(), "LevelEditor.Tabs.Cinematics");
+
+	OutColor = FLinearColor::White;
+	return Icon;
+}
+
+void UMovieGraphTmvEncoderNode::PostEditChangeProperty(FPropertyChangedEvent& PropertyChangedEvent)
+{
+	Super::PostEditChangeProperty(PropertyChangedEvent);
+
+	if (PropertyChangedEvent.GetPropertyName() == GET_MEMBER_NAME_CHECKED(UMovieGraphTmvEncoderNode, TmvEncoderOptions))
+	{
+		OnNodeChangedDelegate.Broadcast(this);
+	}
+}
+#endif // WITH_EDITOR
+
+TUniquePtr<MovieRenderGraph::IVideoCodecWriter> UMovieGraphTmvEncoderNode::Initialize_GameThread(const FMovieGraphVideoNodeInitializationContext& InInitializationContext)
+{
+	bool bIncludeCDOs = true;
+	constexpr bool bExactMatch = true;
+	UMovieGraphGlobalOutputSettingNode* OutputSetting =
+		InInitializationContext.EvaluatedConfig->GetSettingForBranch<UMovieGraphGlobalOutputSettingNode>(GlobalsPinName, bIncludeCDOs, bExactMatch);
+
+	bIncludeCDOs = false;
+	const UMovieGraphTmvEncoderNode* EvaluatedNode = Cast<UMovieGraphTmvEncoderNode>(
+		InInitializationContext.EvaluatedConfig->GetSettingForBranch(GetClass(), FName(InInitializationContext.PassData->Key.RootBranchName), bIncludeCDOs, bExactMatch));
+	checkf(EvaluatedNode, TEXT("TMV Encoder node could not be found in the graph in branch [%s]."), *InInitializationContext.PassData->Key.RootBranchName.ToString());
+	
+	const FFrameRate SourceFrameRate = InInitializationContext.Pipeline->GetDataSourceInstance()->GetDisplayRate();
+	const FFrameRate EffectiveFrameRate = UMovieGraphBlueprintLibrary::GetEffectiveFrameRate(OutputSetting, SourceFrameRate);
+
+	TUniquePtr<FTmvCodecWriter> NewWriter = MakeUnique<FTmvCodecWriter>();
+	
+	FTmvMediaTranscodeJobSettings JobSettings;
+
+	// todo: probably need to rework this.
+	JobSettings.OutputPath.Path = FPaths::GetPath(InInitializationContext.FileName);
+	JobSettings.OutputBaseName = FPaths::GetBaseFilename(InInitializationContext.FileName);
+	if (!JobSettings.OutputBaseName.EndsWith(TEXT(".")))
+	{
+		JobSettings.OutputBaseName += TEXT(".");
+	}
+	JobSettings.ZeroPadFrameNumbers = OutputSetting->ZeroPadFrameNumbers;
+	JobSettings.FrameRate = EffectiveFrameRate;
+	JobSettings.bEnableMipMapping = true;
+
+	NewWriter->TranscodeJob.Reset(NewObject<UTmvMediaTranscodeJob>(GetTransientPackage(), NAME_None));
+	NewWriter->TranscodeJob->Settings = JobSettings;
+	NewWriter->TranscodeJob->AddStage(NewObject<UTmvMediaFrameProducer>(NewWriter->TranscodeJob.Get(), NAME_None)); // Use a dummy frame producer.
+	NewWriter->TranscodeJob->AddStage(NewObject<UTmvMediaFrameConverter>(NewWriter->TranscodeJob.Get(), NAME_None));
+	NewWriter->TranscodeJob->AddStage(NewObject<UTmvMediaTmvFrameEncoder>(NewWriter->TranscodeJob.Get(), NAME_None));
+	NewWriter->TranscodeJob->AddStage(NewObject<UTmvMediaTranscodeMuxer>(NewWriter->TranscodeJob.Get(), NAME_None));
+	NewWriter->TranscodeJob->GetStage<UTmvMediaFrameEncoder>()->SetEncoderOptions(EvaluatedNode->TmvEncoderOptions);
+	NewWriter->TranscodeJob->Start(FApp::GetCurrentTime());
+
+	CachedPipeline = InInitializationContext.Pipeline;
+	
+	return NewWriter;
+}
+
+bool UMovieGraphTmvEncoderNode::Initialize_EncodeThread(MovieRenderGraph::IVideoCodecWriter* InWriter)
+{
+	return true;
+}
+
+void UMovieGraphTmvEncoderNode::WriteFrame_EncodeThread(MovieRenderGraph::IVideoCodecWriter* InWriter, FImagePixelData* InPixelData, TArray<FMovieGraphPassData>&& InCompositePasses, TObjectPtr<UMovieGraphEvaluatedConfig> InEvaluatedConfig, const FString& InBranchName)
+{
+	FTmvCodecWriter* CodecWriter = static_cast<FTmvCodecWriter*>(InWriter);
+	if (!CodecWriter->TranscodeJob)
+	{
+		return;
+	}
+
+	bool bIncludeCDOs = false;
+	constexpr bool bExactMatch = true;
+	const UMovieGraphTmvEncoderNode* EvaluatedNode = Cast<UMovieGraphTmvEncoderNode>(
+		InEvaluatedConfig->GetSettingForBranch(GetClass(), FName(InBranchName), bIncludeCDOs, bExactMatch));
+	checkf(EvaluatedNode, TEXT("TMV Encoder node could not be found in the graph in branch [%s]."), *InBranchName);
+
+	bIncludeCDOs = true;
+	const UMovieGraphGlobalOutputSettingNode* OutputSettingNode =
+		InEvaluatedConfig->GetSettingForBranch<UMovieGraphGlobalOutputSettingNode>(GlobalsPinName, bIncludeCDOs, bExactMatch);
+	
+	const UE::MovieGraph::FMovieGraphSampleState* GraphPayload = InPixelData->GetPayload<UE::MovieGraph::FMovieGraphSampleState>();
+	
+	TArray<FPixelPreProcessor> PixelPreProcessors;
+
+	// Run OCIO before quantizing. For highest accuracy, OCIO should run on the non-quantized pixel data.
+#if WITH_OCIO
+	if (FMovieGraphOCIOHelper::GenerateOcioPixelPreProcessor(GraphPayload, CachedPipeline.Get(), InEvaluatedConfig, EvaluatedNode->OCIOConfiguration, EvaluatedNode->OCIOContext, PixelPreProcessors))
+	{
+		PixelPreProcessors.Pop()(InPixelData);
+	}
+#endif
+
+	// Intermediate processing will be float16 (linear) for TMV color conversion.
+	constexpr int32 TargetBitDepth = 16;
+	constexpr bool bConvertToSrgb = false;
+	TUniquePtr<FImagePixelData> QuantizedPixelData = UE::MoviePipeline::QuantizeImagePixelDataToBitDepth(InPixelData, TargetBitDepth, nullptr, bConvertToSrgb);
+
+	// Do a quick composite of renders/burn-ins.
+	for (const FMovieGraphPassData& CompositePass : InCompositePasses)
+	{
+		// We don't need to copy the data here (even though it's being passed to a async system) because we already made a unique copy of the
+		// burn in/widget data when we decided to composite it.
+		switch (QuantizedPixelData->GetType())
+		{
+		case EImagePixelType::Color:
+			PixelPreProcessors.Add(TAsyncCompositeImage<FColor>(CompositePass.Value->MoveImageDataToNew()));
+			break;
+		case EImagePixelType::Float16:
+			PixelPreProcessors.Add(TAsyncCompositeImage<FFloat16Color>(CompositePass.Value->MoveImageDataToNew()));
+			break;
+		case EImagePixelType::Float32:
+			PixelPreProcessors.Add(TAsyncCompositeImage<FLinearColor>(CompositePass.Value->MoveImageDataToNew()));
+			break;
+		}
+	}
+
+	// This is done on the current thread for simplicity but the composite itself is parallelized.
+	FImagePixelData* PixelData = QuantizedPixelData.Get();
+	for (const FPixelPreProcessor& PreProcessor : PixelPreProcessors)
+	{
+		// PreProcessors are assumed to be valid.
+		PreProcessor(PixelData);
+	}
+	
+	if (UTmvMediaFrameConverter* FrameConverter = CodecWriter->TranscodeJob->GetStage<UTmvMediaFrameConverter>())
+	{
+		TUniquePtr<FTmvMediaFrameMips> FrameMips = MakeUnique<FTmvMediaFrameMips>();
+
+		FrameMips->TimeInfo.FrameIndex = 0;
+		if (const UE::MovieGraph::FMovieGraphSampleState* Payload = InPixelData->GetPayload<UE::MovieGraph::FMovieGraphSampleState>())
+		{
+			FrameMips->TimeInfo.FrameIndex = Payload->TraversalContext.Time.RootFrameNumber.Value;
+		}
+
+		FrameMips->TimeInfo.FrameDuration = FTimespan::FromSeconds(CodecWriter->TranscodeJob->Settings.FrameRate.AsInterval());
+		FrameMips->TimeInfo.FrameTime = FrameMips->TimeInfo.FrameDuration * FrameMips->TimeInfo.FrameIndex;
+		UE::TmvMediaUtils::Encoder::InitFrameMips(*FrameMips, MoveTemp(QuantizedPixelData), CodecWriter->TranscodeJob->Settings.bEnableMipMapping);
+		FrameConverter->ReceiveMips(CodecWriter->TranscodeJob.Get(), MoveTemp(FrameMips));
+	}
+}
+
+void UMovieGraphTmvEncoderNode::BeginFinalize_EncodeThread(MovieRenderGraph::IVideoCodecWriter* InWriter)
+{
+	const FTmvCodecWriter* CodecWriter = static_cast<FTmvCodecWriter*>(InWriter);
+	if(CodecWriter && CodecWriter->TranscodeJob)
+	{
+		CodecWriter->TranscodeJob->RequestStop(FApp::GetCurrentTime(), ETmvMediaTranscodeJobStopReason::Completed);
+	}
+}
+
+void UMovieGraphTmvEncoderNode::Finalize_EncodeThread(MovieRenderGraph::IVideoCodecWriter* InWriter)
+{
+	FTmvCodecWriter* CodecWriter = static_cast<FTmvCodecWriter*>(InWriter);
+	if(CodecWriter && CodecWriter->TranscodeJob)
+	{
+		CodecWriter->TranscodeJob->Discard(FApp::GetCurrentTime());
+		CodecWriter->TranscodeJob.Reset();
+	}
+}
+
+const TCHAR* UMovieGraphTmvEncoderNode::GetFilenameExtension() const
+{
+	if (TmvEncoderOptions.IsValid())
+	{
+		FileExtension = TmvEncoderOptions.Get().GetFileSequenceExtension();
+		return *FileExtension;
+	}
+
+	return TEXT("tmv");
+}
+
+bool UMovieGraphTmvEncoderNode::IsAudioSupported() const
+{
+	// todo: revisit when audio is implemented in Tmv container.
+	return false;
+}
+
+void UMovieGraphTmvEncoderNode::UpdateTelemetry(FMoviePipelineShotRenderTelemetry* InTelemetry) const
+{
+	// Todo: add bUsesTmv?
+	//InTelemetry->bUsesTmv = true;
+}
+
+void UMovieGraphTmvEncoderNode::InjectNodesPostEvaluation(const FName& InBranchName, UMovieGraphEvaluatedConfig* InEvaluatedConfig, TArray<UMovieGraphSettingNode*>& OutInjectedNodes)
+{
+	if (bEnableBurnIn)
+	{
+		UMovieGraphOutputBurnInNode* BurnInNode = NewObject<UMovieGraphOutputBurnInNode>(InEvaluatedConfig);
+		BurnInNode->bOverride_OutputName = true;
+		BurnInNode->bOverride_FileNameFormat = true;
+		BurnInNode->bOverride_OutputRestriction = true;
+		BurnInNode->bOverride_LayerNameFormat = true;
+		BurnInNode->bOverride_BurnInClass = true;
+		BurnInNode->bOverride_bCompositeOntoFinalImage = true;
+		
+		if (TmvEncoderOptions.IsValid())
+		{
+			BurnInNode->OutputName = FString::Printf(TEXT("TMV(%s)"), *TmvEncoderOptions.Get().GetEncoderName().ToString());
+		}
+		else
+		{
+			BurnInNode->OutputName = TEXT("TMV");
+		}
+
+		BurnInNode->FileNameFormat = BurnInFileNameFormat;
+		BurnInNode->OutputRestriction = FSoftClassPath(GetClass());
+		BurnInNode->LayerNameFormat = TEXT("");	// Unused
+		BurnInNode->BurnInClass = BurnInClass;
+		BurnInNode->bCompositeOntoFinalImage = bCompositeOntoFinalImage;
+
+		OutInjectedNodes.Add(BurnInNode);
+	}
+}
\ No newline at end of file
diff --git a/Engine/Plugins/Media/TmvMediaUtils/Source/MovieRenderPipelineTmvEncoder/Private/Graph/MovieGraphTmvEncoderNode.h b/Engine/Plugins/Media/TmvMediaUtils/Source/MovieRenderPipelineTmvEncoder/Private/Graph/MovieGraphTmvEncoderNode.h
new file mode 100644
--- /dev/null
+++ b/Engine/Plugins/Media/TmvMediaUtils/Source/MovieRenderPipelineTmvEncoder/Private/Graph/MovieGraphTmvEncoderNode.h
@@ -0,0 +1,132 @@
+// Copyright Epic Games, Inc. All Rights Reserved.
+
+#pragma once
+
+#include "Graph/MovieGraphNode.h"
+#include "Graph/Nodes/MovieGraphVideoOutputNode.h"
+
+#include "MovieGraphTmvEncoderNode.generated.h"
+
+struct FTmvMediaEncoderOptions;
+class UTmvMediaTranscodeJob;
+
+#define UE_API MOVIERENDERPIPELINETMVENCODER_API
+
+/** A node which can output TMV (Tiled-Mipmap Video) files. */
+UCLASS(MinimalAPI, BlueprintType)
+class UMovieGraphTmvEncoderNode final : public UMovieGraphVideoOutputNode, public IMovieGraphEvaluationNodeInjector
+{
+	GENERATED_BODY()
+
+public:
+	UE_API UMovieGraphTmvEncoderNode();
+
+#if WITH_EDITOR
+	UE_API virtual FText GetNodeTitle(const bool bGetDescriptive = false) const override;
+	UE_API virtual FText GetMenuCategory() const override;
+	UE_API virtual FText GetKeywords() const override;
+	UE_API virtual FLinearColor GetNodeTitleColor() const override;
+	UE_API virtual FSlateIcon GetIconAndTint(FLinearColor& OutColor) const override;
+
+	//~ Begin UObject interface
+	UE_API virtual void PostEditChangeProperty(FPropertyChangedEvent& PropertyChangedEvent) override;
+	//~ End UObject interface
+#endif
+
+protected:
+	// UMovieGraphVideoOutputNode Interface
+	UE_API virtual TUniquePtr<MovieRenderGraph::IVideoCodecWriter> Initialize_GameThread(const FMovieGraphVideoNodeInitializationContext& InInitializationContext) override;
+	UE_API virtual bool Initialize_EncodeThread(MovieRenderGraph::IVideoCodecWriter* InWriter) override;
+	UE_API virtual void WriteFrame_EncodeThread(MovieRenderGraph::IVideoCodecWriter* InWriter, FImagePixelData* InPixelData, TArray<FMovieGraphPassData>&& InCompositePasses, TObjectPtr<UMovieGraphEvaluatedConfig> InEvaluatedConfig, const FString& InBranchName) override;
+	UE_API virtual void BeginFinalize_EncodeThread(MovieRenderGraph::IVideoCodecWriter* InWriter) override;
+	UE_API virtual void Finalize_EncodeThread(MovieRenderGraph::IVideoCodecWriter* InWriter) override;
+	UE_API virtual const TCHAR* GetFilenameExtension() const override;
+	UE_API virtual bool IsAudioSupported() const override;
+	// ~UMovieGraphVideoOutputNode Interface
+
+	// UMovieGraphSettingNode Interface
+	UE_API virtual void UpdateTelemetry(FMoviePipelineShotRenderTelemetry* InTelemetry) const override;
+	// ~UMovieGraphSettingNode Interface
+
+	// IMovieGraphEvaluationNodeInjector Interface
+	virtual void InjectNodesPostEvaluation(const FName& InBranchName, UMovieGraphEvaluatedConfig* InEvaluatedConfig, TArray<UMovieGraphSettingNode*>& OutInjectedNodes) override;
+	// ~IMovieGraphEvaluationNodeInjector Interface
+
+protected:
+	struct FTmvCodecWriter : public MovieRenderGraph::IVideoCodecWriter
+	{
+		/** Transcode job that will convert and encode the produced frames. */
+		TStrongObjectPtr<UTmvMediaTranscodeJob> TranscodeJob;
+	};
+	
+	/** The pipeline that is running this node. */
+	TWeakObjectPtr<UMovieGraphPipeline> CachedPipeline;
+
+	/** Backstore for the currently selected file extensions, which depends on the currently selected TMV encoder. */ 
+	mutable FString FileExtension;
+
+public:
+	UPROPERTY(EditAnywhere, BlueprintReadWrite, Category = Overrides, meta = (InlineEditConditionToggle))
+	uint8 bOverride_TmvEncoderOptions : 1;
+
+	UPROPERTY(EditAnywhere, BlueprintReadWrite, Category = Overrides, meta = (InlineEditConditionToggle))
+	uint8 bOverride_OCIOConfiguration : 1;
+
+	UPROPERTY(EditAnywhere, BlueprintReadWrite, Category = Overrides, meta = (InlineEditConditionToggle))
+	uint8 bOverride_OCIOContext : 1;
+
+	UPROPERTY(EditAnywhere, BlueprintReadWrite, Category = Overrides, meta = (InlineEditConditionToggle))
+	uint8 bOverride_bEnableBurnIn : 1;
+
+	UPROPERTY(EditAnywhere, BlueprintReadWrite, Category = Overrides, meta = (InlineEditConditionToggle))
+	uint8 bOverride_BurnInClass : 1;
+
+	UPROPERTY(EditAnywhere, BlueprintReadWrite, Category = Overrides, meta = (InlineEditConditionToggle))
+	uint8 bOverride_bCompositeOntoFinalImage : 1;
+
+	UPROPERTY(EditAnywhere, BlueprintReadWrite, Category = Overrides, meta = (InlineEditConditionToggle))
+	uint8 bOverride_BurnInFileNameFormat : 1;
+
+	/** TMV Encoder options */
+	UPROPERTY(BlueprintReadWrite, EditAnywhere, Category = "TMV", meta = (ExcludeBaseStruct, EditCondition = "bOverride_TmvEncoderOptions"))
+	TInstancedStruct<FTmvMediaEncoderOptions> TmvEncoderOptions;
+
+	/**
+	* OCIO configuration/transform settings.
+	*
+	* Note: There are differences from the previous implementation in MRQ given that we are now doing CPU-side processing.
+	* 1) This feature only works on desktop platforms when the OpenColorIO library is available.
+	* 2) Users are now responsible for setting the renderer output space to Final Color (HDR) in Linear Working Color Space (SCS_FinalColorHDR) by
+	*    disabling the Tone Curve setting on the renderer node.
+	*/
+	UPROPERTY(EditAnywhere, BlueprintReadWrite, Category = "OCIO", DisplayName="OCIO Configuration", meta = (DisplayAfter = "FileNameFormat", EditCondition = "bOverride_OCIOConfiguration"))
+	FOpenColorIODisplayConfiguration OCIOConfiguration;
+
+	/**
+	* OCIO context of key-value string pairs, typically used to apply shot-specific looks (such as a CDL color correction, or a 1D grade LUT).
+	* 
+	* Notes:
+	* 1) If a configuration asset base context was set, it remains active but can be overridden here with new key-values.
+	* 2) Format tokens such as {shot_name} are supported and will get resolved before submission.
+	*/
+	UPROPERTY(EditAnywhere, BlueprintReadWrite, Category = "OCIO", DisplayName = "OCIO Context", meta = (DisplayAfter = "OCIOConfiguration", EditCondition = "bOverride_OCIOContext"))
+	TMap<FString, FString> OCIOContext;
+
+	/** If true, this output node will also generate a burn-in. */
+	UPROPERTY(EditAnywhere, BlueprintReadWrite, Category = "Burn In", meta = (EditCondition = "bOverride_bEnableBurnIn"))
+	bool bEnableBurnIn = false;
+
+	/** The widget that the burn-in should use. */
+	UPROPERTY(EditAnywhere, BlueprintReadWrite, Category="Burn In", meta=(MetaClass="/Script/MovieRenderPipelineCore.MovieGraphBurnInWidget", EditCondition="bOverride_BurnInClass"))
+	FSoftClassPath BurnInClass = FSoftClassPath(TEXT("/MovieRenderPipeline/Blueprints/Graph/DefaultGraphBurnIn.DefaultGraphBurnIn_C"));
+	
+	/** If true, the burn-in that's generated will be composited onto this output. Otherwise, the burn-in will be written to a different file. */
+	UPROPERTY(EditAnywhere, BlueprintReadWrite, Category = "Burn In", meta = (EditCondition = "bOverride_bCompositeOntoFinalImage"))
+	bool bCompositeOntoFinalImage = true;
+
+	/** The file name format used for writing out the burn-in video (if not composited). */
+	UPROPERTY(EditAnywhere, BlueprintReadWrite, Category = "Burn In", meta = (EditCondition = "bOverride_BurnInFileNameFormat", MrgTokenAutocomplete))
+	FString BurnInFileNameFormat = TEXT("{sequence_name}.{layer_name}.{renderer_name}.{frame_number}");
+};
+
+#undef UE_API
diff --git a/Engine/Plugins/Media/TmvMediaUtils/Source/MovieRenderPipelineTmvEncoder/Private/MovieRenderPipelineTmvEncoderModule.cpp b/Engine/Plugins/Media/TmvMediaUtils/Source/MovieRenderPipelineTmvEncoder/Private/MovieRenderPipelineTmvEncoderModule.cpp
new file mode 100644
--- /dev/null
+++ b/Engine/Plugins/Media/TmvMediaUtils/Source/MovieRenderPipelineTmvEncoder/Private/MovieRenderPipelineTmvEncoderModule.cpp
@@ -0,0 +1,22 @@
+// Copyright Epic Games, Inc. All Rights Reserved.
+
+#include "Modules/ModuleManager.h"
+
+class FMovieRenderPipelineTmvEncoderModule : public IModuleInterface
+{
+public:
+	/*~ Begin IModuleInterface interface */
+	virtual void StartupModule() override;
+	virtual void ShutdownModule() override;
+	/*~ End IModuleInterface interface */
+};
+
+void FMovieRenderPipelineTmvEncoderModule::StartupModule()
+{
+}
+
+void FMovieRenderPipelineTmvEncoderModule::ShutdownModule()
+{
+}
+
+IMPLEMENT_MODULE(FMovieRenderPipelineTmvEncoderModule, MovieRenderPipelineTmvEncoder)
\ No newline at end of file
diff --git a/Engine/Plugins/Media/TmvMediaUtils/Source/MovieRenderPipelineTmvEncoder/Private/SampleConverter/TmvMediaFrameMipPixelDataBuffer.cpp b/Engine/Plugins/Media/TmvMediaUtils/Source/MovieRenderPipelineTmvEncoder/Private/SampleConverter/TmvMediaFrameMipPixelDataBuffer.cpp
new file mode 100644
--- /dev/null
+++ b/Engine/Plugins/Media/TmvMediaUtils/Source/MovieRenderPipelineTmvEncoder/Private/SampleConverter/TmvMediaFrameMipPixelDataBuffer.cpp
@@ -0,0 +1,71 @@
+// Copyright Epic Games, Inc. All Rights Reserved.
+
+#include "TmvMediaFrameMipPixelDataBuffer.h"
+
+#include "ImageCore.h"
+#include "ImagePixelData.h"
+#include "RHIResources.h"
+#include "Utils/TmvMediaFrameUtils.h"
+
+namespace UE::TmvMediaUtils
+{
+	void* GetPixelDataPointer(const FImagePixelData& InImage)
+	{
+		const void* Data = nullptr;
+		int64 DataSize;
+		if (InImage.GetRawData(Data, DataSize))
+		{
+			return const_cast<void*>(Data);
+		}
+		return nullptr;
+	}
+}
+
+FTmvMediaFrameMipPixelDataBuffer::FTmvMediaFrameMipPixelDataBuffer() = default;
+
+FTmvMediaFrameMipPixelDataBuffer::FTmvMediaFrameMipPixelDataBuffer(TUniquePtr<FImagePixelData>&& InImage, int32 InMipLevel)
+{
+	if (InImage)
+	{
+		FImageView ImageView = InImage->GetImageView();
+		AllocatedBufferSize = ImageView.GetImageSizeBytes();
+		UE::TmvMedia::FrameUtils::PopulateMipInfoFromImageInfo(InMipLevel, ImageView, MipInfo);
+		PlaneImages.Add(MoveTemp(InImage));
+	}
+}
+
+FTmvMediaFrameMipPixelDataBuffer::~FTmvMediaFrameMipPixelDataBuffer() = default;
+
+bool FTmvMediaFrameMipPixelDataBuffer::RequestAllocation(const FTmvMediaFrameMipInfo& InMipInfo)
+{
+	// todo: currently not used. Implement if needed.
+	ensureMsgf(false, TEXT("FTmvMediaFrameMipPixelDataBuffer::RequestAllocation is not implemented."));
+	return false;
+}
+
+void* FTmvMediaFrameMipPixelDataBuffer::GetMappedBuffer()
+{
+	// Only return the first plane buffer if we have one plane to be compliant with
+	// the api that should return one contiguous buffer.
+	if (PlaneImages.Num() == 1 && PlaneImages.IsValidIndex(0) && PlaneImages[0])
+	{
+		return UE::TmvMediaUtils::GetPixelDataPointer(*PlaneImages[0]);
+	}
+	return nullptr;
+}
+
+void* FTmvMediaFrameMipPixelDataBuffer::GetPlaneBufferForComponent(int32 InComponentIndex)
+{
+	// Use the mip info to find the plane index for this component
+	const int32 PlaneIndex = MipInfo.GetPlaneIndexForComponent(InComponentIndex);
+	if (PlaneImages.IsValidIndex(PlaneIndex) && PlaneImages[PlaneIndex])
+	{
+		return UE::TmvMediaUtils::GetPixelDataPointer(*PlaneImages[PlaneIndex]);
+	}
+	return nullptr;
+}
+
+FShaderResourceViewRHIRef FTmvMediaFrameMipPixelDataBuffer::GetShaderResourceView(int32 InComponentIndex) const
+{
+	return nullptr;
+}
\ No newline at end of file
diff --git a/Engine/Plugins/Media/TmvMediaUtils/Source/MovieRenderPipelineTmvEncoder/Private/SampleConverter/TmvMediaFrameMipPixelDataBuffer.h b/Engine/Plugins/Media/TmvMediaUtils/Source/MovieRenderPipelineTmvEncoder/Private/SampleConverter/TmvMediaFrameMipPixelDataBuffer.h
new file mode 100644
--- /dev/null
+++ b/Engine/Plugins/Media/TmvMediaUtils/Source/MovieRenderPipelineTmvEncoder/Private/SampleConverter/TmvMediaFrameMipPixelDataBuffer.h
@@ -0,0 +1,42 @@
+// Copyright Epic Games, Inc. All Rights Reserved.
+
+#pragma once
+
+#include "Containers/Array.h"
+#include "SampleConverter/TmvMediaFrameMipBuffer.h"
+
+struct FImagePixelData;
+
+/**
+ * Implementation of texture cpu buffering backed by FImagePixelData.
+ */
+class FTmvMediaFrameMipPixelDataBuffer : public FTmvMediaFrameMipBuffer
+{
+public:
+	/* Constructor */
+	FTmvMediaFrameMipPixelDataBuffer();
+
+	/** Construct from an existing image, using the image as backing buffer. */
+	FTmvMediaFrameMipPixelDataBuffer(TUniquePtr<FImagePixelData>&& InImage, int32 InMipLevel);
+
+	/** Destructor to clean up the internal buffer data. */
+	virtual ~FTmvMediaFrameMipPixelDataBuffer() override;
+
+	//~ Begin FTmvMediaFrameMipBuffer
+	virtual bool RequestAllocation(const FTmvMediaFrameMipInfo& InMipInfo) override;
+	virtual void WaitAllocation() override {}
+	virtual bool TryUpdateMipInfo(const FTmvMediaFrameMipInfo& InMipInfo) override { return false; }
+	virtual void* GetMappedBuffer() override;
+	virtual void* GetPlaneBufferForComponent(int32 InComponentIndex) override;
+	virtual FShaderResourceViewRHIRef GetShaderResourceView(int32 InComponentIndex) const override;
+	virtual void CopyTileRegions(int32 InFrameId, const FTmvMediaFrameMipInfo& InTmvMipInfo, TConstArrayView<FIntRect> InTileRegions, TConstArrayView<FTmvMediaShaderTileDesc> InTileInfos) override {}
+	virtual bool IsValidForRendering() const override
+	{
+		return false;	// This buffer type can't be used for rendering.
+	}
+	//~ End FTmvMediaFrameMipBuffer
+
+private:
+	/** Image buffer for each plane. */
+	TArray<TUniquePtr<FImagePixelData>> PlaneImages;
+};
\ No newline at end of file
diff --git a/Engine/Plugins/Media/TmvMediaUtils/TmvMediaUtils.uplugin b/Engine/Plugins/Media/TmvMediaUtils/TmvMediaUtils.uplugin
new file mode 100644
--- /dev/null
+++ b/Engine/Plugins/Media/TmvMediaUtils/TmvMediaUtils.uplugin
@@ -0,0 +1,37 @@
+{
+	"FileVersion" : 3,
+	"Version" : 1,
+	"VersionName" : "1.0",
+	"FriendlyName" : "Tiled Mipmap Media Player Utilities",
+	"Description" : "Implements a media player for tiled-mipmap video.",
+	"Category" : "Media Players",
+	"CreatedBy" : "Epic Games, Inc.",
+	"CreatedByURL" : "https://epicgames.com",
+	"DocsURL" : "",
+	"MarketplaceURL" : "",
+	"SupportURL" : "",
+	"EnabledByDefault" : false,
+	"CanContainContent" : false,
+	"IsBetaVersion" : true,
+	"Installed" : false,
+	"Plugins": [
+		{
+			"Name": "TmvMedia",
+			"Enabled": true,
+			"Optional": false
+		},
+		{
+			"Name": "MovieRenderPipeline",
+			"Enabled": true,
+			"Optional": false
+		}
+	],
+	"Modules" :
+	[
+		{
+			"Name" : "MovieRenderPipelineTmvEncoder",
+			"Type" : "Runtime",
+			"LoadingPhase" : "Default"
+		}
+	]
+}
