[TmvMedia] Updating plugin description to make it easier to find "tmv" plugins.

#Synced-CL 51094493

diff --git a/Engine/Plugins/Media/TmvMedia/TmvMedia.uplugin b/Engine/Plugins/Media/TmvMedia/TmvMedia.uplugin
--- a/Engine/Plugins/Media/TmvMedia/TmvMedia.uplugin
+++ b/Engine/Plugins/Media/TmvMedia/TmvMedia.uplugin
@@ -2,8 +2,8 @@
 	"FileVersion" : 3,
 	"Version" : 1,
 	"VersionName" : "1.0",
-	"FriendlyName" : "Tiled Mipmap Media Player",
-	"Description" : "Implements a media player for tiled-mipmap image sequences.",
+	"FriendlyName" : "Tiled Mipmap Video Player",
+	"Description" : "Framework for tiled-mipmap video (TMV) playback, includes transcoding tools.\nImplemented using Advanced Professional Video (APV) codec.",
 	"Category" : "Media Players",
 	"CreatedBy" : "Epic Games, Inc.",
 	"CreatedByURL" : "https://epicgames.com",
diff --git a/Engine/Plugins/Media/TmvMediaUtils/TmvMediaUtils.uplugin b/Engine/Plugins/Media/TmvMediaUtils/TmvMediaUtils.uplugin
--- a/Engine/Plugins/Media/TmvMediaUtils/TmvMediaUtils.uplugin
+++ b/Engine/Plugins/Media/TmvMediaUtils/TmvMediaUtils.uplugin
@@ -2,8 +2,8 @@
 	"FileVersion" : 3,
 	"Version" : 1,
 	"VersionName" : "1.0",
-	"FriendlyName" : "Tiled Mipmap Media Player Utilities",
-	"Description" : "Implements a media player for tiled-mipmap video.",
+	"FriendlyName" : "Tiled Mipmap Video Utilities",
+	"Description" : "Additional tools for the tiled-mipmap video (TMV) framework.\nImplements Movie Render Graph (MRG) integration.",
 	"Category" : "Media Players",
 	"CreatedBy" : "Epic Games, Inc.",
 	"CreatedByURL" : "https://epicgames.com",
