[TmvMedia] Cleaning up pass on module dependencies.

#Synced-CL 51094493

diff --git a/Engine/Plugins/Media/TmvMedia/Source/ApvMedia/ApvMedia.Build.cs b/Engine/Plugins/Media/TmvMedia/Source/ApvMedia/ApvMedia.Build.cs
--- a/Engine/Plugins/Media/TmvMedia/Source/ApvMedia/ApvMedia.Build.cs
+++ b/Engine/Plugins/Media/TmvMedia/Source/ApvMedia/ApvMedia.Build.cs
@@ -6,15 +6,11 @@ namespace UnrealBuildTool.Rules
 	{
 		public ApvMedia(ReadOnlyTargetRules Target) : base(Target)
 		{
-			PublicDependencyModuleNames.AddRange(
-				new string[] {
-					"Core",
-					"CoreUObject",
-				});
-
 			PrivateDependencyModuleNames.AddRange(
 				new string[]
 				{
+					"Core",
+					"CoreUObject",
 					"DeveloperSettings",
 					"Engine",
 					"Media",
diff --git a/Engine/Plugins/Media/TmvMedia/Source/TmvMedia/TmvMedia.Build.cs b/Engine/Plugins/Media/TmvMedia/Source/TmvMedia/TmvMedia.Build.cs
--- a/Engine/Plugins/Media/TmvMedia/Source/TmvMedia/TmvMedia.Build.cs
+++ b/Engine/Plugins/Media/TmvMedia/Source/TmvMedia/TmvMedia.Build.cs
@@ -6,28 +6,28 @@ namespace UnrealBuildTool.Rules
 	{
 		public TmvMedia(ReadOnlyTargetRules Target) : base(Target)
 		{
-			PrivateDependencyModuleNames.AddRange(
-				new string[]
-				{
-					"Renderer",
-				});
-
 			PublicDependencyModuleNames.AddRange(
 				new string[] {
 					"Core",
 					"CoreUObject",
 					"Engine",
+					"Media",
+					"RHI",
+					"TmvMediaShaders",
+				});
+
+			PrivateDependencyModuleNames.AddRange(
+				new string[]
+				{
 					"ImageCore",
 					"ImageWrapper",
 					"Json",
 					"JsonUtilities",
-					"Media",
 					"MediaAssets",
 					"MediaUtils",
-					"Projects", 
-					"RHI",
+					"Projects",
+					"Renderer",
 					"RenderCore",
-					"TmvMediaShaders",
 				});
 		}
 	}
diff --git a/Engine/Plugins/Media/TmvMediaUtils/Source/MovieRenderPipelineTmvEncoder/MovieRenderPipelineTmvEncoder.Build.cs b/Engine/Plugins/Media/TmvMediaUtils/Source/MovieRenderPipelineTmvEncoder/MovieRenderPipelineTmvEncoder.Build.cs
--- a/Engine/Plugins/Media/TmvMediaUtils/Source/MovieRenderPipelineTmvEncoder/MovieRenderPipelineTmvEncoder.Build.cs
+++ b/Engine/Plugins/Media/TmvMediaUtils/Source/MovieRenderPipelineTmvEncoder/MovieRenderPipelineTmvEncoder.Build.cs
@@ -6,25 +6,20 @@ namespace UnrealBuildTool.Rules
 	{
 		public MovieRenderPipelineTmvEncoder(ReadOnlyTargetRules Target) : base(Target)
 		{
-            PublicDependencyModuleNames.AddRange(				
-				new string[] {
-					"MovieRenderPipelineCore",
-					"MovieRenderPipelineSettings",
-					"OpenColorIO",
-					"SlateCore",
-                });
-
 			PrivateDependencyModuleNames.AddRange(
 				new string[]
 				{
 					"Core",
-                    "CoreUObject",
-                    "Engine",
-                    "ImageCore",
-                    "RHI",
-                    "TmvMedia",
-                }
-            );
-        }
-    }
+					"CoreUObject",
+					"Engine",
+					"ImageCore",
+					"MovieRenderPipelineCore",
+					"OpenColorIO",
+					"RHI",
+					"SlateCore",
+					"TmvMedia",
+				}
+			);
+		}
+	}
 }
